Fixed Undo History being an unrecognized tab on restart

diff --git a/Engine/Source/Editor/UndoHistory/Private/UndoHistoryModule.cpp b/Engine/Source/Editor/UndoHistory/Private/UndoHistoryModule.cpp
--- a/Engine/Source/Editor/UndoHistory/Private/UndoHistoryModule.cpp
+++ b/Engine/Source/Editor/UndoHistory/Private/UndoHistoryModule.cpp
@@ -8,16 +8,16 @@
 
 void FUndoHistoryModule::StartupModule()
 {
-	FGlobalTabmanager::Get()->RegisterTabSpawner(UndoHistoryTabName, FOnSpawnTab::CreateRaw(this, &FUndoHistoryModule::HandleSpawnSettingsTab))
+	FGlobalTabmanager::Get()->RegisterNomadTabSpawner(UndoHistoryTabName, FOnSpawnTab::CreateRaw(this, &FUndoHistoryModule::HandleSpawnSettingsTab))
 		.SetDisplayName(NSLOCTEXT("FUndoHistoryModule", "UndoHistoryTabTitle", "Undo History"))
 		.SetTooltipText(NSLOCTEXT("FUndoHistoryModule", "UndoHistoryTooltipText", "Open the Undo History tab."))
-		.SetIcon(FSlateIcon(FEditorStyle::GetStyleSetName(), "UndoHistory.TabIcon"));
+		.SetIcon(FSlateIcon(FEditorStyle::GetStyleSetName(), "UndoHistory.TabIcon"))
+		.SetAutoGenerateMenuEntry(false);
 }
 
 void FUndoHistoryModule::ShutdownModule()
 {
-	FGlobalTabmanager::Get()->UnregisterTabSpawner(UndoHistoryTabName);
-	//FGlobalTabmanager::Get()->UnregisterNomadTabSpawner(UndoHistoryTabName);
+	FGlobalTabmanager::Get()->UnregisterNomadTabSpawner(UndoHistoryTabName);
 }
 
 bool FUndoHistoryModule::SupportsDynamicReloading()
@@ -25,11 +25,6 @@ bool FUndoHistoryModule::SupportsDynamicReloading()
 	return true;
 }
 
-//void FUndoHistoryModule::ExecuteOpenUndoHistory() /*const*/
-//{
-//	FGlobalTabmanager::Get()->InvokeTab(UndoHistoryTabName);
-//}
-
 TSharedRef<SDockTab> FUndoHistoryModule::HandleSpawnSettingsTab(const FSpawnTabArgs& SpawnTabArgs)
 {
 	const TSharedRef<SDockTab> DockTab = SNew(SDockTab)
diff --git a/Engine/Source/Editor/UndoHistory/Public/UndoHistoryModule.h b/Engine/Source/Editor/UndoHistory/Public/UndoHistoryModule.h
--- a/Engine/Source/Editor/UndoHistory/Public/UndoHistoryModule.h
+++ b/Engine/Source/Editor/UndoHistory/Public/UndoHistoryModule.h
@@ -21,8 +21,7 @@ public:
 
 	// End IModuleInterface interface
 
-	//static /*virtual */void ExecuteOpenUndoHistory() /*const override*/;
-	static void ExecuteOpenUndoHistory() /*const*/
+	static void ExecuteOpenUndoHistory()
 	{
 		FGlobalTabmanager::Get()->InvokeTab(UndoHistoryTabName);
 	}
diff --git a/Engine/Source/Runtime/Slate/Private/Framework/Docking/TabManager.cpp b/Engine/Source/Runtime/Slate/Private/Framework/Docking/TabManager.cpp
--- a/Engine/Source/Runtime/Slate/Private/Framework/Docking/TabManager.cpp
+++ b/Engine/Source/Runtime/Slate/Private/Framework/Docking/TabManager.cpp
@@ -775,14 +775,20 @@ void FTabManager::PopulateTabSpawnerMenu( FMenuBuilder& PopulateMe, TSharedRef<F
 		for ( FTabSpawner::TIterator SpawnerIterator(TabSpawner); SpawnerIterator; ++SpawnerIterator  )	
 		{
 			const TSharedRef<FTabSpawnerEntry>& SpawnerEntry = SpawnerIterator.Value();
-			AllSpawners->AddUnique(SpawnerEntry);
+			if (SpawnerEntry->bAutoGenerateMenuEntry)
+			{
+				AllSpawners->AddUnique( SpawnerEntry );
+			}
 		}
 
 		// General Tabs
 		for ( FTabSpawner::TIterator SpawnerIterator(*NomadTabSpawner); SpawnerIterator; ++SpawnerIterator  )	
 		{
 			const TSharedRef<FTabSpawnerEntry>& SpawnerEntry = SpawnerIterator.Value();
-			AllSpawners->AddUnique(SpawnerEntry);
+			if (SpawnerEntry->bAutoGenerateMenuEntry)
+			{
+				AllSpawners->AddUnique( SpawnerEntry );
+			}
 		}
 	}
 
diff --git a/Engine/Source/Runtime/Slate/Public/Framework/Docking/TabManager.h b/Engine/Source/Runtime/Slate/Public/Framework/Docking/TabManager.h
--- a/Engine/Source/Runtime/Slate/Public/Framework/Docking/TabManager.h
+++ b/Engine/Source/Runtime/Slate/Public/Framework/Docking/TabManager.h
@@ -121,6 +121,7 @@ struct FTabSpawnerEntry : public FWorkspaceItem
 		, OnFindTabToReuse()
 		, MenuType(ETabSpawnerMenuType::Display)
 		, SpawnedTabPtr()
+		, bAutoGenerateMenuEntry(true)
 
 	{
 	}
@@ -161,6 +162,12 @@ struct FTabSpawnerEntry : public FWorkspaceItem
 		return *this;
 	}
 
+	FTabSpawnerEntry& SetAutoGenerateMenuEntry( bool bInAutoGenerateMenuEntry )
+	{
+		bAutoGenerateMenuEntry = bInAutoGenerateMenuEntry;
+		return *this;
+	}
+
 	virtual TSharedPtr<FTabSpawnerEntry> AsSpawnerEntry() override
 	{
 		return SharedThis(this);
@@ -172,6 +179,8 @@ private:
 	/** When this method is not provided, we assume that the tab should only allow 0 or 1 instances */
 	FOnFindTabToReuse OnFindTabToReuse;
 	ETabSpawnerMenuType::Type MenuType;
+	/** Whether to automatically generate a menu entry for this tab spawner */
+	bool bAutoGenerateMenuEntry;
 
 	TWeakPtr<SDockTab> SpawnedTabPtr;
 
