Check if index is invalid

diff --git a/Engine/Source/Editor/UndoHistory/Private/Widgets/SUndoHistory.cpp b/Engine/Source/Editor/UndoHistory/Private/Widgets/SUndoHistory.cpp
--- a/Engine/Source/Editor/UndoHistory/Private/Widgets/SUndoHistory.cpp
+++ b/Engine/Source/Editor/UndoHistory/Private/Widgets/SUndoHistory.cpp
@@ -108,7 +108,10 @@ void SUndoHistory::OnUndoBufferChanged()
 		ReloadUndoList();
 
 		LastActiveTransactionIndex = GEditor->Trans->GetQueueLength() - 1;
-		UndoListView->SetSelection(UndoList[LastActiveTransactionIndex]);
+		if (UndoList.IsValidIndex(LastActiveTransactionIndex))
+		{
+			UndoListView->SetSelection(UndoList[LastActiveTransactionIndex]);
+		}
 	}
 }
 
