Copied 4366928 Jeremie.Roy
Fix crashes related to the UndoHistory.
Add missing UndoBufferChangedDelegate broadcasts in EditorTransaction

diff --git a/Engine/Source/Editor/UndoHistory/Private/Widgets/SUndoHistory.cpp b/Engine/Source/Editor/UndoHistory/Private/Widgets/SUndoHistory.cpp
--- a/Engine/Source/Editor/UndoHistory/Private/Widgets/SUndoHistory.cpp
+++ b/Engine/Source/Editor/UndoHistory/Private/Widgets/SUndoHistory.cpp
@@ -103,21 +103,9 @@ SUndoHistory::~SUndoHistory()
 
 void SUndoHistory::OnUndoBufferChanged()
 {
-	if (GEditor != nullptr && GEditor->Trans != nullptr)
-	{
-		ReloadUndoList();
-
-		LastActiveTransactionIndex = GEditor->Trans->GetQueueLength() - 1;
+	ReloadUndoList();
 
-		if (UndoList.IsValidIndex(LastActiveTransactionIndex))
-		{
-			UndoListView->SetSelection(UndoList[LastActiveTransactionIndex]);
-		}
-		else
-		{
-			UndoListView->ClearSelection();
-		}
-	}
+	SelectLastTransaction();
 }
 
 void SUndoHistory::Tick( const FGeometry& AllottedGeometry, const double InCurrentTime, const float InDeltaTime )
@@ -155,12 +143,24 @@ void SUndoHistory::ReloadUndoList( )
 
 	for (int32 QueueIndex = 0; QueueIndex < GEditor->Trans->GetQueueLength(); ++QueueIndex)
 	{
-		UndoList.Add(MakeShareable(new FTransactionInfo(QueueIndex, GEditor->Trans->GetTransaction(QueueIndex))));
+		UndoList.Add(MakeShared<FTransactionInfo>(QueueIndex, GEditor->Trans->GetTransaction(QueueIndex)));
 	}
 
 	UndoListView->RequestListRefresh();
 }
 
+void SUndoHistory::SelectLastTransaction()
+{
+	if (GEditor != nullptr && GEditor->Trans != nullptr)
+	{
+		LastActiveTransactionIndex = GEditor->Trans->GetQueueLength() - 1;
+		if (UndoList.IsValidIndex(LastActiveTransactionIndex))
+		{
+			UndoListView->SetSelection(UndoList[LastActiveTransactionIndex]);
+		}
+	}
+}
+
 
 /* SUndoHistory callbacks
  *****************************************************************************/
diff --git a/Engine/Source/Editor/UndoHistory/Private/Widgets/SUndoHistory.h b/Engine/Source/Editor/UndoHistory/Private/Widgets/SUndoHistory.h
--- a/Engine/Source/Editor/UndoHistory/Private/Widgets/SUndoHistory.h
+++ b/Engine/Source/Editor/UndoHistory/Private/Widgets/SUndoHistory.h
@@ -67,6 +67,9 @@ protected:
 
 private:
 
+	// Select the last transaction in the undo history.
+	void SelectLastTransaction();
+
 	// Callback for clicking the 'Discard History' button.
 	FReply HandleDiscardHistoryButtonClicked( );
 
diff --git a/Engine/Source/Editor/UnrealEd/Classes/Editor/TransBuffer.h b/Engine/Source/Editor/UnrealEd/Classes/Editor/TransBuffer.h
--- a/Engine/Source/Editor/UnrealEd/Classes/Editor/TransBuffer.h
+++ b/Engine/Source/Editor/UnrealEd/Classes/Editor/TransBuffer.h
@@ -130,6 +130,7 @@ protected:
 
 				GUndo->BeginOperation();
 				TransactionStateChangedDelegate.Broadcast(GUndo->GetContext(), ETransactionStateEventType::TransactionStarted);
+				UndoBufferChangedDelegate.Broadcast();
 			}
 			const int32 PriorRecordsCount = (Result > 0 ? ActiveRecordCounts[Result - 1] : 0);
 			ActiveRecordCounts.Add(UndoBuffer.Last()->GetRecordCount() - PriorRecordsCount);
diff --git a/Engine/Source/Editor/UnrealEd/Private/EditorTransaction.cpp b/Engine/Source/Editor/UnrealEd/Private/EditorTransaction.cpp
--- a/Engine/Source/Editor/UnrealEd/Private/EditorTransaction.cpp
+++ b/Engine/Source/Editor/UnrealEd/Private/EditorTransaction.cpp
@@ -1084,6 +1084,7 @@ int32 UTransBuffer::End()
 				{
 					check(UndoCount == 0);
 					UndoBuffer.Pop(false);
+					UndoBufferChangedDelegate.Broadcast();
 				}
 			}
 			GUndo = nullptr;
@@ -1127,6 +1128,7 @@ void UTransBuffer::Reset( const FText& Reason )
 		ResetReason = Reason;
 		ActiveCount = 0;
 		ActiveRecordCounts.Empty();
+		UndoBufferChangedDelegate.Broadcast();
 
 		CheckState();
 	}
@@ -1167,7 +1169,6 @@ void UTransBuffer::Cancel( int32 StartIndex /*=0*/ )
 
 			UndoCount = PreviousUndoCount;
 			PreviousUndoCount = INDEX_NONE;
-
 			UndoBufferChangedDelegate.Broadcast();
 		}
 		else
@@ -1346,6 +1347,8 @@ bool UTransBuffer::Undo(bool bCanRedo)
 		{
 			UndoBuffer.RemoveAt(UndoBuffer.Num() - UndoCount, UndoCount);
 			UndoCount = 0;
+
+			UndoBufferChangedDelegate.Broadcast();
 		}
 	}
 	GIsTransacting = false;
