Remove usage of raw pointer in UndoHistory in favor of just using the already available QueueIndex to prevent potential stale pointer

diff --git a/Engine/Source/Editor/UndoHistory/Private/Widgets/SUndoHistory.cpp b/Engine/Source/Editor/UndoHistory/Private/Widgets/SUndoHistory.cpp
--- a/Engine/Source/Editor/UndoHistory/Private/Widgets/SUndoHistory.cpp
+++ b/Engine/Source/Editor/UndoHistory/Private/Widgets/SUndoHistory.cpp
@@ -20,6 +20,16 @@
 #define LOCTEXT_NAMESPACE "SUndoHistory"
 
 
+// Helper function for TransactionInfo
+const FTransaction* SUndoHistory::FTransactionInfo::GetTransaction() const
+{
+	if (GEditor && GEditor->Trans)
+	{
+		return GEditor->Trans->GetTransaction(QueueIndex);
+	}
+	return nullptr;
+}
+
 /* SUndoHistory interface
  *****************************************************************************/
 
@@ -269,7 +279,7 @@ void SUndoHistory::ReloadUndoList()
 
 	for (int32 QueueIndex = 0; QueueIndex < GEditor->Trans->GetQueueLength(); ++QueueIndex)
 	{
-		UndoList.Add(MakeShared<FTransactionInfo>(QueueIndex, GEditor->Trans->GetTransaction(QueueIndex)));
+		UndoList.Add(MakeShared<FTransactionInfo>(QueueIndex));
 	}
 
 	UndoListView->RequestListRefresh();
@@ -283,11 +293,11 @@ void SUndoHistory::SelectLastTransaction()
 		if (UndoList.IsValidIndex(LastActiveTransactionIndex))
 		{
 			const TSharedPtr<FTransactionInfo>& TransactionInfo = UndoList[LastActiveTransactionIndex];
-
-			if (TransactionInfo.IsValid() && TransactionInfo->Transaction)
+			const FTransaction* InnerTransaction = TransactionInfo.IsValid() ? TransactionInfo->GetTransaction() : nullptr;
+			if (InnerTransaction)
 			{
 				UndoListView->SetSelection(TransactionInfo);
-				UndoDetailsView->SetSelectedTransaction(TransactionInfo->Transaction->GenerateDiff());
+				UndoDetailsView->SetSelectedTransaction(InnerTransaction->GenerateDiff());
 			}
 		}
 	}
@@ -316,19 +326,18 @@ TSharedRef<ITableRow> SUndoHistory::HandleUndoListGenerateRow(TSharedPtr<FTransa
 		.OnGotoTransactionClicked(this, &SUndoHistory::HandleGoToTransaction)
 		.IsApplied(this, &SUndoHistory::HandleUndoListRowIsApplied, TransactionInfo->QueueIndex)
 		.QueueIndex(TransactionInfo->QueueIndex)
-		.Transaction(TransactionInfo->Transaction);
+		.Transaction(TransactionInfo->GetTransaction());
 }
 
 void SUndoHistory::HandleGoToTransaction(const FGuid& TargetTransactionId)
 {
 	TSharedPtr<FTransactionInfo>* Transaction = UndoList.FindByPredicate([TargetTransactionId](TSharedPtr<FTransactionInfo> TransactionInfo)
 	{
-		if (TransactionInfo->Transaction == nullptr)
+		if (const FTransaction* InnerTransaction = TransactionInfo->GetTransaction())
 		{
-			return false;
+			return InnerTransaction->GetId() == TargetTransactionId;
 		}
-
-		return TransactionInfo->Transaction->GetId() == TargetTransactionId;
+		return false;
 	});
 
 	if (Transaction != nullptr)
@@ -355,9 +364,9 @@ void SUndoHistory::HandleUndoListSelectionChanged(TSharedPtr<FTransactionInfo> I
 		return;
 	}
 
-	if (InItem->Transaction)
+	if (const FTransaction* InnerTransaction = InItem->GetTransaction())
 	{
-		UndoDetailsView->SetSelectedTransaction(InItem->Transaction->GenerateDiff());
+		UndoDetailsView->SetSelectedTransaction(InnerTransaction->GenerateDiff());
 	}
 
 	if (SelectInfo == ESelectInfo::OnMouseClick || SelectInfo == ESelectInfo::Direct)
@@ -486,4 +495,4 @@ bool SUndoHistory::IsShowingTransactionDetails() const
 	return GetDefault<UUndoHistorySettings>()->bShowTransactionDetails;
 }
 
-#undef LOCTEXT_NAMESPACE
+#undef LOCTEXT_NAMESPACE
\ No newline at end of file
diff --git a/Engine/Source/Editor/UndoHistory/Private/Widgets/SUndoHistory.h b/Engine/Source/Editor/UndoHistory/Private/Widgets/SUndoHistory.h
--- a/Engine/Source/Editor/UndoHistory/Private/Widgets/SUndoHistory.h
+++ b/Engine/Source/Editor/UndoHistory/Private/Widgets/SUndoHistory.h
@@ -24,14 +24,12 @@ class SUndoHistory
 		// Holds the transactions index in the transaction queue.
 		int32 QueueIndex;
 
-		// Holds a pointer to the transaction.
-		const FTransaction* Transaction;
-
 		// Creates and initializes a new instance.
-		FTransactionInfo( int32 InQueueIndex, const FTransaction* InTransaction )
-			: QueueIndex(InQueueIndex)
-			, Transaction(InTransaction)
+		FTransactionInfo( int32 InQueueIndex )
+			: QueueIndex(InQueueIndex)			
 		{ }
+
+		const FTransaction* GetTransaction() const;
 	};
 
 public:
diff --git a/Engine/Source/Editor/UndoHistory/Private/Widgets/SUndoHistoryTableRow.h b/Engine/Source/Editor/UndoHistory/Private/Widgets/SUndoHistoryTableRow.h
--- a/Engine/Source/Editor/UndoHistory/Private/Widgets/SUndoHistoryTableRow.h
+++ b/Engine/Source/Editor/UndoHistory/Private/Widgets/SUndoHistoryTableRow.h
@@ -53,22 +53,23 @@ public:
 	{
 		IsApplied = InArgs._IsApplied;
 		QueueIndex = InArgs._QueueIndex;
-		TransactionId = InArgs._Transaction->GetId();
 		OnGotoTransactionClicked = InArgs._OnGotoTransactionClicked;
 
 		FSuperRowType::FArguments Args = FSuperRowType::FArguments()
 			.Style(&FEditorStyle::Get().GetWidgetStyle<FTableRowStyle>("SceneOutliner.TableViewRow"));
 
-		UObject* ContextObject = InArgs._Transaction->GetContext().PrimaryObject;
+		UObject* ContextObject = nullptr;
+		if (InArgs._Transaction)
+		{
+			TransactionId = InArgs._Transaction->GetId();
+			ContextObject = InArgs._Transaction->GetContext().PrimaryObject;
+			Title = InArgs._Transaction->GetTitle();
+		}
 
 		if (ContextObject != nullptr)
 		{
 			Title = FText::Format(LOCTEXT("UndoHistoryTableRowTitleF", "{0} [{1}]"), InArgs._Transaction->GetTitle(), FText::FromString(ContextObject->GetFName().ToString()));
 		}
-		else
-		{
-			Title = InArgs._Transaction->GetTitle();
-		}
 
 		SMultiColumnTableRow<TSharedPtr<int32> >::Construct(Args, InOwnerTableView);
 	}
