Fixing unstyled text in the undo history

diff --git a/Engine/Source/Editor/UndoHistory/Private/Widgets/SUndoHistory.cpp b/Engine/Source/Editor/UndoHistory/Private/Widgets/SUndoHistory.cpp
--- a/Engine/Source/Editor/UndoHistory/Private/Widgets/SUndoHistory.cpp
+++ b/Engine/Source/Editor/UndoHistory/Private/Widgets/SUndoHistory.cpp
@@ -69,7 +69,6 @@ void SUndoHistory::Construct( const FArguments& InArgs )
 										.HAlign(HAlign_Left)
 										[
 											SNew(STextBlock)
-												.TextStyle( FEditorStyle::Get(), "Docking.TabFont" )
 												.Text(LOCTEXT("TransactionHistoryLabel", "Transactions"))
 										]
 								]
